Horde issue 881859

diff --git a/Engine/Plugins/Experimental/PSDImporterMaterialDesignerBridge/Source/PSDImporterMaterialDesignerBridge/Private/PSDImporterMDContentBrowserIntegration.cpp b/Engine/Plugins/Experimental/PSDImporterMaterialDesignerBridge/Source/PSDImporterMaterialDesignerBridge/Private/PSDImporterMDContentBrowserIntegration.cpp
--- a/Engine/Plugins/Experimental/PSDImporterMaterialDesignerBridge/Source/PSDImporterMaterialDesignerBridge/Private/PSDImporterMDContentBrowserIntegration.cpp
+++ b/Engine/Plugins/Experimental/PSDImporterMaterialDesignerBridge/Source/PSDImporterMaterialDesignerBridge/Private/PSDImporterMDContentBrowserIntegration.cpp
@@ -89,22 +89,22 @@ void FPSDImporterMaterialDesignerContentBrowserIntegration::CreateMenuEntries(FM
 	InMenuBuilder.BeginSection(TEXT("PSDDocument"), LOCTEXT("PSDDocumentMaterialDesigner", "PSD Document (Material Designer)"));
 
 	InMenuBuilder.AddMenuEntry(
-		LOCTEXT("CreatePSDDocumentMateriaMaterialDesignerl", "Create Material From PSD Document (Material Designer)"),
-		LOCTEXT("CreatePSDDocumentMaterialMaterialDesignerTooltip", "Creates a material comprising of the layer stack in the PSD Document using a Material Designer Material."),
+		LOCTEXT("CreatePSDDocumentMateriaMaterialDesigner", "Create Material From PSD Document (Material Designer)"),
+		LOCTEXT("CreatePSDDocumentMateriaMaterialDesignerTooltip", "Creates a material comprising of the layer stack in the PSD Document using a Material Designer Material."),
 		FSlateIconFinder::FindIconForClass(UMaterial::StaticClass()),
 		FUIAction(FExecuteAction::CreateRaw(this, &FPSDImporterMaterialDesignerContentBrowserIntegration::CreatePSDMaterialMaterialDesigner, InSelectedAssets))
 	);
 
 	InMenuBuilder.AddMenuEntry(
-		LOCTEXT("CreatePSDDocumentQuadsMaterialDesigner", "Create Quads From PSD Document (Material Designer (Instance))"),
-		LOCTEXT("CreatePSDDocumentQuadsMaterialDesignerTooltip", "Creates a series of quads representing each layer in the PSD Document using instances of a Material Designer Material."),
+		LOCTEXT("CreatePSDDocumentQuadsMaterialDesignerInstance", "Create Quads From PSD Document (Material Designer (Instance))"),
+		LOCTEXT("CreatePSDDocumentQuadsMaterialDesignerInstanceTooltip", "Creates a series of quads representing each layer in the PSD Document using instances of a Material Designer Material."),
 		FSlateIconFinder::FindIconForClass(UStaticMesh::StaticClass()),
 		FUIAction(FExecuteAction::CreateRaw(this, &FPSDImporterMaterialDesignerContentBrowserIntegration::CreatePSDQuadsMaterialDesigner, InSelectedAssets, EPSDImporterMaterialDesignerType::Instance))
 	);
 
 	InMenuBuilder.AddMenuEntry(
-		LOCTEXT("CreatePSDDocumentQuadsMaterialDesigner", "Create Quads From PSD Document (Material Designer (Material))"),
-		LOCTEXT("CreatePSDDocumentQuadsMaterialDesignerTooltip", "Creates a series of quads representing each layer in the PSD Document using new Material Designer Materials."),
+		LOCTEXT("CreatePSDDocumentQuadsMaterialDesignerMaterial", "Create Quads From PSD Document (Material Designer (Material))"),
+		LOCTEXT("CreatePSDDocumentQuadsMaterialDesignerMaterialTooltip", "Creates a series of quads representing each layer in the PSD Document using new Material Designer Materials."),
 		FSlateIconFinder::FindIconForClass(UStaticMesh::StaticClass()),
 		FUIAction(FExecuteAction::CreateRaw(this, &FPSDImporterMaterialDesignerContentBrowserIntegration::CreatePSDQuadsMaterialDesigner, InSelectedAssets, EPSDImporterMaterialDesignerType::Copy))
 	);
